[UE] Add experimental Global Configuration Data

#tests Ran unit tests, tested text version of blueprint function in graphs and view bindings

diff --git a/Engine/Build/LowLevelTests/GlobalConfigurationData.xml b/Engine/Build/LowLevelTests/GlobalConfigurationData.xml
new file mode 100644
--- /dev/null
+++ b/Engine/Build/LowLevelTests/GlobalConfigurationData.xml
@@ -0,0 +1,8 @@
+<?xml version="1.0" encoding="utf-8"?>
+<BuildGraph xmlns:xsi="http://www.w3.org/2001/XMLSchema-instance" xsi:schemaLocation="http://www.epicgames.com/BuildGraph ../../Build/Graph/Schema.xsd" xmlns="http://www.epicgames.com/BuildGraph">
+  <Option Name="RunGlobalConfigurationDataTests" DefaultValue="" Description="Run GlobalConfigurationData Tests" />
+  <Property Name="TestNames" Value="$(TestNames);GlobalConfigurationData" />
+  <Extend Name="RunAllTests">
+    <Expand Name="DeployAndTest" Platform="Win64" TestName="GlobalConfigurationData" ShortName="GlobalConfigurationData" TargetName="GlobalConfigurationDataTests" BinaryRelativePath="Engine\Binaries" ReportType="console" />
+  </Extend>
+</BuildGraph>
\ No newline at end of file
diff --git a/Engine/Plugins/Experimental/GlobalConfigurationData/GlobalConfigurationData.uplugin b/Engine/Plugins/Experimental/GlobalConfigurationData/GlobalConfigurationData.uplugin
new file mode 100644
--- /dev/null
+++ b/Engine/Plugins/Experimental/GlobalConfigurationData/GlobalConfigurationData.uplugin
@@ -0,0 +1,28 @@
+{
+	"FileVersion" : 3,
+	"Version" : 1,
+	"VersionName" : "0.1",
+	"FriendlyName" : "Global Configuration Data",
+	"Description" : "A system that is used to query configuration data that can come from many different sources without knowing specifically which one.",
+	"Category": "Misc",
+	"CreatedBy" : "Epic Games, Inc.",
+	"CreatedByURL" : "http://epicgames.com",
+	"DocsURL" : "",
+	"MarketplaceURL" : "",
+	"SupportURL" : "",
+	"ExplicitlyLoaded": false,
+	"EnabledByDefault" : true,
+	"CanContainContent" : false,
+	"IsBetaVersion" : false,
+	"IsExperimentalVersion": true,
+	"Installed" : false,
+	"Modules": [
+		{
+			"Name": "GlobalConfigurationData",
+			"Type": "Runtime",
+			"LoadingPhase": "Default"
+		}
+	],
+	"Plugins": [
+	]
+}
diff --git a/Engine/Plugins/Experimental/GlobalConfigurationData/Source/GlobalConfigurationData/GlobalConfigurationData.Build.cs b/Engine/Plugins/Experimental/GlobalConfigurationData/Source/GlobalConfigurationData/GlobalConfigurationData.Build.cs
new file mode 100644
--- /dev/null
+++ b/Engine/Plugins/Experimental/GlobalConfigurationData/Source/GlobalConfigurationData/GlobalConfigurationData.Build.cs
@@ -0,0 +1,23 @@
+// Copyright Epic Games, Inc. All Rights Reserved.
+
+namespace UnrealBuildTool.Rules
+{
+	public class GlobalConfigurationData : ModuleRules
+	{
+        public GlobalConfigurationData(ReadOnlyTargetRules Target) : base(Target)
+		{
+			PCHUsage = ModuleRules.PCHUsageMode.UseExplicitOrSharedPCHs;
+
+			PublicDependencyModuleNames.AddRange(
+				new string[]
+				{
+					"Core",
+					"CoreUObject",
+					"Engine",
+					"Json",
+					"JsonUtilities",
+				}
+			);
+		}
+	}
+}
diff --git a/Engine/Plugins/Experimental/GlobalConfigurationData/Source/GlobalConfigurationData/Private/GlobalConfigurationData.cpp b/Engine/Plugins/Experimental/GlobalConfigurationData/Source/GlobalConfigurationData/Private/GlobalConfigurationData.cpp
new file mode 100644
--- /dev/null
+++ b/Engine/Plugins/Experimental/GlobalConfigurationData/Source/GlobalConfigurationData/Private/GlobalConfigurationData.cpp
@@ -0,0 +1,72 @@
+// Copyright Epic Games, Inc. All Rights Reserved.
+
+#include "GlobalConfigurationData.h"
+
+#include "GlobalConfigurationDataInternal.h"
+#include "GlobalConfigurationRouter.h"
+#include "JsonObjectConverter.h"
+
+bool UE::GlobalConfigurationData::TryGetDataOfType(const FString& EntryName, const UStruct* Type, void* DataOut)
+{
+	if (TSharedPtr<FJsonValue> Data = IGlobalConfigurationRouter::TryGetData(EntryName))
+	{
+		TSharedPtr<FJsonObject>* JsonObject;
+		if (Data->TryGetObject(JsonObject))
+		{
+			return FJsonObjectConverter::JsonObjectToUStruct(JsonObject->ToSharedRef(), Type, DataOut);
+		}
+		else
+		{
+			UE_LOGFMT(LogGlobalConfigurationData, Error, "Failed to json deserialize key {EntryName} from value {EntryValue}", *EntryName, IGlobalConfigurationRouter::TryPrintString(Data));
+		}
+	}
+	return false;
+}
+
+bool UE::GlobalConfigurationData::TryGetData(const FString& EntryName, bool& bValueOut)
+{
+	if (TSharedPtr<FJsonValue> Data = IGlobalConfigurationRouter::TryGetData(EntryName))
+	{
+		return Data->TryGetBool(bValueOut);
+	}
+	return false;
+}
+
+bool UE::GlobalConfigurationData::TryGetData(const FString& EntryName, int32& ValueOut)
+{
+	if (TSharedPtr<FJsonValue> Data = IGlobalConfigurationRouter::TryGetData(EntryName))
+	{
+		return Data->TryGetNumber(ValueOut);
+	}
+	return false;
+}
+
+bool UE::GlobalConfigurationData::TryGetData(const FString& EntryName, float& ValueOut)
+{
+	if (TSharedPtr<FJsonValue> Data = IGlobalConfigurationRouter::TryGetData(EntryName))
+	{
+		return Data->TryGetNumber(ValueOut);
+	}
+	return false;
+}
+
+bool UE::GlobalConfigurationData::TryGetData(const FString& EntryName, FString& ValueOut)
+{
+	if (TSharedPtr<FJsonValue> Data = IGlobalConfigurationRouter::TryGetData(EntryName))
+	{
+		return Data->TryGetString(ValueOut);
+	}
+	return false;
+}
+
+bool UE::GlobalConfigurationData::TryGetData(const FString& EntryName, FText& ValueOut)
+{
+	FString StringData;
+	if (TryGetData(EntryName, StringData))
+	{
+		// Todo: Better interpretation for loc text
+		ValueOut = FText::FromString(MoveTemp(StringData));
+		return true;
+	}
+	return false;
+}
\ No newline at end of file
diff --git a/Engine/Plugins/Experimental/GlobalConfigurationData/Source/GlobalConfigurationData/Private/GlobalConfigurationDataBlueprintLibrary.cpp b/Engine/Plugins/Experimental/GlobalConfigurationData/Source/GlobalConfigurationData/Private/GlobalConfigurationDataBlueprintLibrary.cpp
new file mode 100644
--- /dev/null
+++ b/Engine/Plugins/Experimental/GlobalConfigurationData/Source/GlobalConfigurationData/Private/GlobalConfigurationDataBlueprintLibrary.cpp
@@ -0,0 +1,75 @@
+// Copyright Epic Games, Inc. All Rights Reserved.
+
+#include "GlobalConfigurationDataBlueprintLibrary.h"
+#include "GlobalConfigurationData.h"
+
+#include UE_INLINE_GENERATED_CPP_BY_NAME(GlobalConfigurationDataBlueprintLibrary)
+
+bool UGlobalConfigurationDataBlueprintLibrary::GetConfigDataBool(const FString& EntryName, bool &bValueOut)
+{
+	return UE::GlobalConfigurationData::TryGetData(EntryName, bValueOut);
+}
+
+bool UGlobalConfigurationDataBlueprintLibrary::GetConfigDataInt(const FString& EntryName, int32& ValueOut)
+{
+	return UE::GlobalConfigurationData::TryGetData(EntryName, ValueOut);
+}
+
+bool UGlobalConfigurationDataBlueprintLibrary::GetConfigDataFloat(const FString& EntryName, float& ValueOut)
+{
+	return UE::GlobalConfigurationData::TryGetData(EntryName, ValueOut);
+}
+
+bool UGlobalConfigurationDataBlueprintLibrary::GetConfigDataString(const FString& EntryName, FString& ValueOut)
+{
+	return UE::GlobalConfigurationData::TryGetData(EntryName, ValueOut);
+}
+
+bool UGlobalConfigurationDataBlueprintLibrary::GetConfigDataText(const FString& EntryName, FText& ValueOut)
+{
+	return UE::GlobalConfigurationData::TryGetData(EntryName, ValueOut);
+}
+
+bool UGlobalConfigurationDataBlueprintLibrary::GetConfigDataStruct(const FString& EntryName, UScriptStruct* StructType, FInstancedStruct& ValueOut)
+{
+	if (StructType)
+	{
+		ValueOut.InitializeAs(StructType);
+		return UE::GlobalConfigurationData::TryGetDataOfType(EntryName, ValueOut.GetScriptStruct(), ValueOut.GetMutableMemory());
+	}
+	return false;
+}
+
+bool UGlobalConfigurationDataBlueprintLibrary::GetConfigDataObject(const FString& EntryName, UObject* ValueInOut)
+{
+	if (ValueInOut)
+	{
+		return UE::GlobalConfigurationData::TryGetDataOfType(EntryName, ValueInOut->GetClass(), ValueInOut);
+	}
+	return false;
+}
+
+bool UGlobalConfigurationDataBlueprintLibrary::GetConfigDataBoolWithDefault(const FString& EntryName, bool bDefaultValue)
+{
+	return UE::GlobalConfigurationData::GetDataWithDefault(EntryName, bDefaultValue);
+}
+
+int32 UGlobalConfigurationDataBlueprintLibrary::GetConfigDataIntWithDefault(const FString& EntryName, int32 DefaultValue)
+{
+	return UE::GlobalConfigurationData::GetDataWithDefault(EntryName, DefaultValue);
+}
+
+float UGlobalConfigurationDataBlueprintLibrary::GetConfigDataFloatWithDefault(const FString& EntryName, float DefaultValue)
+{
+	return UE::GlobalConfigurationData::GetDataWithDefault(EntryName, DefaultValue);
+}
+
+FString UGlobalConfigurationDataBlueprintLibrary::GetConfigDataStringWithDefault(const FString& EntryName, FString DefaultValue)
+{
+	return UE::GlobalConfigurationData::GetDataWithDefault(EntryName, DefaultValue);
+}
+
+FText UGlobalConfigurationDataBlueprintLibrary::GetConfigDataTextWithDefault(const FString& EntryName, FText DefaultValue)
+{
+	return UE::GlobalConfigurationData::GetDataWithDefault(EntryName, DefaultValue);
+}
diff --git a/Engine/Plugins/Experimental/GlobalConfigurationData/Source/GlobalConfigurationData/Private/GlobalConfigurationDataInternal.h b/Engine/Plugins/Experimental/GlobalConfigurationData/Source/GlobalConfigurationData/Private/GlobalConfigurationDataInternal.h
new file mode 100644
--- /dev/null
+++ b/Engine/Plugins/Experimental/GlobalConfigurationData/Source/GlobalConfigurationData/Private/GlobalConfigurationDataInternal.h
@@ -0,0 +1,6 @@
+// Copyright Epic Games, Inc. All Rights Reserved.
+#pragma once
+
+#include "Logging/StructuredLog.h"
+
+DECLARE_LOG_CATEGORY_EXTERN(LogGlobalConfigurationData, Log, All);
diff --git a/Engine/Plugins/Experimental/GlobalConfigurationData/Source/GlobalConfigurationData/Private/GlobalConfigurationDataModule.cpp b/Engine/Plugins/Experimental/GlobalConfigurationData/Source/GlobalConfigurationData/Private/GlobalConfigurationDataModule.cpp
new file mode 100644
--- /dev/null
+++ b/Engine/Plugins/Experimental/GlobalConfigurationData/Source/GlobalConfigurationData/Private/GlobalConfigurationDataModule.cpp
@@ -0,0 +1,21 @@
+// Copyright Epic Games, Inc. All Rights Reserved.
+
+#include "Modules/ModuleManager.h"
+
+#include "GlobalConfigurationDataInternal.h"
+#include "Logging/LogMacros.h"
+#include "Routers/GlobalConfigurationConfigRouter.h"
+#include "Routers/GlobalConfigurationConsoleCommandRouter.h"
+
+DEFINE_LOG_CATEGORY(LogGlobalConfigurationData);
+
+class FGlobalConfigurationDataModule : public IModuleInterface
+{
+private:
+	
+	// Just existing is enough for these
+	FGlobalConfigurationConfigRouter ConfigRouter;
+	FGlobalConfigurationConsoleCommandRouter ConsoleCommandRouter;
+};
+
+IMPLEMENT_MODULE(FGlobalConfigurationDataModule, GlobalConfigurationData);
diff --git a/Engine/Plugins/Experimental/GlobalConfigurationData/Source/GlobalConfigurationData/Private/GlobalConfigurationRouter.cpp b/Engine/Plugins/Experimental/GlobalConfigurationData/Source/GlobalConfigurationData/Private/GlobalConfigurationRouter.cpp
new file mode 100644
--- /dev/null
+++ b/Engine/Plugins/Experimental/GlobalConfigurationData/Source/GlobalConfigurationData/Private/GlobalConfigurationRouter.cpp
@@ -0,0 +1,169 @@
+// Copyright Epic Games, Inc. All Rights Reserved.
+
+#include "GlobalConfigurationRouter.h"
+#include "GlobalConfigurationDataInternal.h"
+
+#include "HAL/IConsoleManager.h"
+#include "JsonObjectConverter.h"
+#include "Policies/CondensedJsonPrintPolicy.h"
+#include "Serialization/JsonSerializer.h"
+
+namespace IGlobalConfigurationRouter_Private
+{
+	TArray<const IGlobalConfigurationRouter*> RegisteredRouters;
+
+	class FStackJsonReader : public TJsonReader<>
+	{
+	public:
+		FStackJsonReader(const FStringView& String)
+		: TJsonReader(&Reader)
+		, Reader((void*)String.GetData(), String.NumBytes(), false)
+		{
+		}
+		FBufferReader Reader;
+	};
+
+	class FStackJsonWriter : public TJsonStringWriter<TCondensedJsonPrintPolicy<TCHAR>>
+	{
+	public:
+		explicit FStackJsonWriter(FString* Out)
+		: TJsonStringWriter(Out, 0)
+		{
+		}
+	};
+
+	static bool bShowParseErrorsAsWarning = false;
+	static FAutoConsoleVariableRef CVarMaxSearchResultsToReport(
+		TEXT("GCD.ShowParseErrorsAsWarning"),
+		bShowParseErrorsAsWarning,
+		TEXT("If enabled, json parse errors will be logged as warnings, otherwise they are logged as verbose."));
+}
+
+IGlobalConfigurationRouter::IGlobalConfigurationRouter(FString&& InRouterName, int32 InPriority)
+: RouterName(MoveTemp(InRouterName))
+, Priority(InPriority)
+{
+	check(IsInGameThread());
+	IGlobalConfigurationRouter_Private::RegisteredRouters.Add(this);
+	IGlobalConfigurationRouter_Private::RegisteredRouters.Sort([](const IGlobalConfigurationRouter& LHS, const IGlobalConfigurationRouter& RHS)
+		{
+			// Can't have routers with the same priority
+			check(LHS.Priority != RHS.Priority);
+			return LHS.Priority > RHS.Priority;
+		});
+
+	UE_LOGFMT(LogGlobalConfigurationData, Verbose, "Registered Global Configuration Router {RouterName}", RouterName);
+}
+IGlobalConfigurationRouter::~IGlobalConfigurationRouter()
+{
+	check(IsInGameThread());
+	IGlobalConfigurationRouter_Private::RegisteredRouters.Remove(this);
+	UE_LOGFMT(LogGlobalConfigurationData, Verbose, "Unregistered Global Configuration Router {RouterName}", RouterName);
+}
+
+TSharedPtr<FJsonValue> IGlobalConfigurationRouter::TryGetData(const FString& FullEntryName)
+{
+	check(IsInGameThread());
+
+	FString ObjectName;
+	FString FieldName;
+
+	const FString& EntryName = FullEntryName.Split(TEXT("."), &ObjectName, &FieldName) ? ObjectName : FullEntryName;
+	UE_CLOGFMT(!FieldName.IsEmpty(), LogGlobalConfigurationData, VeryVerbose, "IGlobalConfigurationRouter::TryGetData parsing full entry name into Entry: {EntryName}, Field: {FieldName}", EntryName, FieldName);
+
+	for (const IGlobalConfigurationRouter* Router : IGlobalConfigurationRouter_Private::RegisteredRouters)
+	{
+		TSharedPtr<FJsonValue> Result = Router->TryGetDataFromRouter(EntryName);
+		if (Result.IsValid())
+		{
+			if (!FieldName.IsEmpty())
+			{
+				if (Result->Type == EJson::Object)
+				{
+					if (TSharedPtr<FJsonValue> Field = Result->AsObject()->TryGetField(FieldName))
+					{
+						UE_LOGFMT(LogGlobalConfigurationData, VeryVerbose, "IGlobalConfigurationRouter::TryGetData found value for {EntryName} in router {RouterName}: {EntryValue}", FullEntryName, Router->RouterName, TryPrintString(Field));
+						return Field;
+					}
+					else
+					{
+						UE_LOGFMT(LogGlobalConfigurationData, Warning, "IGlobalConfigurationRouter::TryGetData found {EntryName} but it does not have field {FieldName}, skipping for router {RouterName}: {EntryValue}", EntryName, FieldName, Router->RouterName, TryPrintString(Result));
+						continue;
+					}
+				}
+				else
+				{
+					UE_LOGFMT(LogGlobalConfigurationData, Warning, "IGlobalConfigurationRouter::TryGetData found {EntryName} but it is not an object that can satisfy field request {FieldName}, skipping for router {RouterName}: {EntryValue}", EntryName, FieldName, Router->RouterName, TryPrintString(Result));
+					continue;
+				}
+			}
+
+			UE_LOGFMT(LogGlobalConfigurationData, VeryVerbose, "IGlobalConfigurationRouter::TryGetData found value for {EntryName} in router {RouterName}: {EntryValue}", FullEntryName, Router->RouterName, TryPrintString(Result));
+			return Result;
+		}
+	}
+
+	// Todo: Fall back to searching full value without looking for an entry in an object?
+
+	UE_LOGFMT(LogGlobalConfigurationData, VeryVerbose, "IGlobalConfigurationRouter::TryGetData found no results for {EntryName}", FullEntryName);
+	return {};
+}
+
+void IGlobalConfigurationRouter::GetAllRegisteredData(TMap<FString, TMap<FString, TSharedRef<FJsonValue>>>& DataOut)
+{
+	check(IsInGameThread());
+
+	TMap<FString, TSharedRef<FJsonValue>> Entries;
+	for (const IGlobalConfigurationRouter* Router : IGlobalConfigurationRouter_Private::RegisteredRouters)
+	{
+		Entries.Empty();
+		Router->GetAllDataFromRouter(Entries);
+
+		for (TPair<FString, TSharedRef<FJsonValue>> Pair : Entries)
+		{
+			DataOut.FindOrAdd(Pair.Key).Add(Router->RouterName, Pair.Value);
+		}
+	}
+}
+
+TSharedPtr<FJsonValue> IGlobalConfigurationRouter::TryParseString(const FStringView& String)
+{
+	if (String.IsEmpty())
+	{
+		return {};
+	}
+
+	TSharedPtr<FJsonValue> JsonValue;
+	IGlobalConfigurationRouter_Private::FStackJsonReader Reader(String);
+	if (FJsonSerializer::Deserialize(Reader, JsonValue))
+	{
+		return JsonValue;
+	}
+	else
+	{
+		if (IGlobalConfigurationRouter_Private::bShowParseErrorsAsWarning)
+		{
+			UE_LOGFMT(LogGlobalConfigurationData, Warning, "IGlobalConfigurationRouter::TryParseString parsing {JsonString} failed with error {JsonError}, storing as raw string value instead.", String, Reader.GetErrorMessage());
+		}
+		else
+		{
+			UE_LOGFMT(LogGlobalConfigurationData, Verbose, "IGlobalConfigurationRouter::TryParseString parsing {JsonString} failed with error {JsonError}, storing as raw string value instead.", String, Reader.GetErrorMessage());
+		}
+	}
+
+	// If regular parse failed then just assume it's an unquoted string
+	return MakeShared<FJsonValueString>(FString(String));
+}
+
+FString IGlobalConfigurationRouter::TryPrintString(TSharedPtr<FJsonValue> Value)
+{
+	if (!Value.IsValid())
+	{
+		return {};
+	}
+
+	FString OutputString;
+	IGlobalConfigurationRouter_Private::FStackJsonWriter Writer(&OutputString);
+	FJsonSerializer::Serialize(Value, FString(), Writer);
+	return OutputString;
+}
diff --git a/Engine/Plugins/Experimental/GlobalConfigurationData/Source/GlobalConfigurationData/Private/Routers/GlobalConfigurationConfigRouter.cpp b/Engine/Plugins/Experimental/GlobalConfigurationData/Source/GlobalConfigurationData/Private/Routers/GlobalConfigurationConfigRouter.cpp
new file mode 100644
--- /dev/null
+++ b/Engine/Plugins/Experimental/GlobalConfigurationData/Source/GlobalConfigurationData/Private/Routers/GlobalConfigurationConfigRouter.cpp
@@ -0,0 +1,56 @@
+// Copyright Epic Games, Inc. All Rights Reserved.
+
+#include "Routers/GlobalConfigurationConfigRouter.h"
+
+#include "GlobalConfigurationDataInternal.h"
+#include "Misc/ConfigCacheIni.h"
+#include "Misc/CoreDelegates.h"
+
+namespace GlobalConfigurationConfigRouter_Private
+{
+	const TCHAR* Section = TEXT("GlobalConfigurationData");
+}
+
+FGlobalConfigurationConfigRouter::FGlobalConfigurationConfigRouter()
+: IGlobalConfigurationRouter(TEXT("Config"), INT32_MIN)
+{
+	OnConfigSectionsChanged({}, {});
+	FCoreDelegates::TSOnConfigSectionsChanged().AddRaw(this, &FGlobalConfigurationConfigRouter::OnConfigSectionsChanged);
+}
+
+FGlobalConfigurationConfigRouter::~FGlobalConfigurationConfigRouter()
+{
+	FCoreDelegates::TSOnConfigSectionsChanged().RemoveAll(this);
+}
+
+TSharedPtr<FJsonValue> FGlobalConfigurationConfigRouter::TryGetDataFromRouter(const FString& EntryName) const
+{
+	if (const TSharedRef<FJsonValue>* Value = StoredData.Find(EntryName))
+	{
+		return *Value;
+	}
+	return {};
+}
+
+void FGlobalConfigurationConfigRouter::GetAllDataFromRouter(TMap<FString, TSharedRef<FJsonValue>>& DataOut) const
+{
+	DataOut = StoredData;
+}
+
+void FGlobalConfigurationConfigRouter::OnConfigSectionsChanged(const FString& IniFilename, const TSet<FString>& SectionNames)
+{
+	if (SectionNames.IsEmpty() || (IniFilename == GEngineIni && SectionNames.Contains(GlobalConfigurationConfigRouter_Private::Section)))
+	{
+		StoredData.Empty();
+
+		FKeyValueSink Visitor = FKeyValueSink::CreateLambda([this](const TCHAR* Key, const TCHAR* ValueString)
+			{
+				if (TSharedPtr<FJsonValue> Value = TryParseString(ValueString))
+				{
+					StoredData.Add(Key, Value.ToSharedRef());
+				}
+			});
+
+		GConfig->ForEachEntry(Visitor, GlobalConfigurationConfigRouter_Private::Section, GEngineIni);
+	}
+}
\ No newline at end of file
diff --git a/Engine/Plugins/Experimental/GlobalConfigurationData/Source/GlobalConfigurationData/Private/Routers/GlobalConfigurationConsoleCommandRouter.cpp b/Engine/Plugins/Experimental/GlobalConfigurationData/Source/GlobalConfigurationData/Private/Routers/GlobalConfigurationConsoleCommandRouter.cpp
new file mode 100644
--- /dev/null
+++ b/Engine/Plugins/Experimental/GlobalConfigurationData/Source/GlobalConfigurationData/Private/Routers/GlobalConfigurationConsoleCommandRouter.cpp
@@ -0,0 +1,69 @@
+// Copyright Epic Games, Inc. All Rights Reserved.
+
+#include "Routers/GlobalConfigurationConsoleCommandRouter.h"
+
+#include "Dom/JsonValue.h"
+#include "GlobalConfigurationDataInternal.h"
+#include "HAL/IConsoleManager.h"
+
+namespace GlobalConfigurationConfigRouter_Private
+{
+	TMap<FString, TSharedRef<FJsonValue>> StoredData;
+
+	void RegisterValue(const TArray<FString>& Args)
+	{
+		if (Args.Num() == 2)
+		{
+			if (TSharedPtr<FJsonValue> Value = FGlobalConfigurationConsoleCommandRouter::TryParseString(Args[1]))
+			{
+				StoredData.Emplace(Args[0], Value.ToSharedRef());
+				UE_LOGFMT(LogGlobalConfigurationData, Verbose, "Registering Console Command value for {EntryName}: {EntryValue}", Args[0], Args[1]);
+			}
+			else
+			{
+				UE_LOGFMT(LogGlobalConfigurationData, Error, "GCD.RegisterValue failed to parse value for {EntryName}: {EntryValue}", Args[0], Args[1]);
+			}
+		}
+		else
+		{
+			UE_LOGFMT(LogGlobalConfigurationData, Verbose, "GCD.RegisterValue called with an incorrect number of arguments, did you forget quotes on the value?");
+		}
+	}
+
+	void UnregisterValue(const TArray<FString>& Args)
+	{
+		if (!Args.IsEmpty())
+		{
+			StoredData.Remove(Args[0]);
+		}
+	}
+
+	FAutoConsoleCommand RegisterValueCommand(
+		TEXT("GCD.RegisterValue"),
+		TEXT("Register a new value into the Global Configuration Data system that will override all others"),
+		FConsoleCommandWithArgsDelegate::CreateStatic(&RegisterValue));
+
+	FAutoConsoleCommand UnregisterValueCommand(
+		TEXT("GCD.UnregisterValue"),
+		TEXT("Remove a value from the Global Configuration Data system that was overriding all others"),
+		FConsoleCommandWithArgsDelegate::CreateStatic(&UnregisterValue));
+}
+
+FGlobalConfigurationConsoleCommandRouter::FGlobalConfigurationConsoleCommandRouter()
+: IGlobalConfigurationRouter(TEXT("ConsoleCommand"), INT32_MAX)
+{
+}
+
+TSharedPtr<FJsonValue> FGlobalConfigurationConsoleCommandRouter::TryGetDataFromRouter(const FString& EntryName) const
+{
+	if (const TSharedRef<FJsonValue>* Value = GlobalConfigurationConfigRouter_Private::StoredData.Find(EntryName))
+	{
+		return *Value;
+	}
+	return {};
+}
+
+void FGlobalConfigurationConsoleCommandRouter::GetAllDataFromRouter(TMap<FString, TSharedRef<FJsonValue>>& DataOut) const
+{
+	DataOut = GlobalConfigurationConfigRouter_Private::StoredData;
+}
diff --git a/Engine/Plugins/Experimental/GlobalConfigurationData/Source/GlobalConfigurationData/Public/GlobalConfigurationData.h b/Engine/Plugins/Experimental/GlobalConfigurationData/Source/GlobalConfigurationData/Public/GlobalConfigurationData.h
new file mode 100644
--- /dev/null
+++ b/Engine/Plugins/Experimental/GlobalConfigurationData/Source/GlobalConfigurationData/Public/GlobalConfigurationData.h
@@ -0,0 +1,51 @@
+// Copyright Epic Games, Inc. All Rights Reserved.
+#pragma once
+
+#include "Containers/UnrealString.h"
+
+class UStruct;
+class UObject;
+
+/**
+ * GlobalConfigurationData is a system that allows querying for data from a unique name.
+ * Multiple data routers can be registered so a user doesn't have to know where it's coming from.
+ * This system allows use to use a multitude of systems to deliver configuration to the user.
+ * 
+ * There are two built-in routers that use config files or console commands. Other routers exist
+ * in other plugins, or users can create their own, to pull data from other sources for use as hotfixes
+ * data experiments, etc.
+ */
+namespace UE::GlobalConfigurationData
+{
+	bool TryGetDataOfType(const FString& EntryName, const UStruct* Type, void* DataOut);
+
+	bool TryGetData(const FString& EntryName, bool& bValueOut);
+	bool TryGetData(const FString& EntryName, int32& ValueOut);
+	bool TryGetData(const FString& EntryName, float& ValueOut);
+	bool TryGetData(const FString& EntryName, FString& ValueOut);
+	bool TryGetData(const FString& EntryName, FText& ValueOut);
+
+	template<typename Type UE_REQUIRES(TModels_V<CStaticStructProvider, Type>)>
+	bool TryGetData(const FString& EntryName, Type& DataOut)
+	{
+		return TryGetDataOfType(EntryName, Type::StaticStruct(), &DataOut);
+	}
+
+	template<typename Type UE_REQUIRES(std::is_base_of_v<UObject, Type>)>
+	bool TryGetData(const FString& EntryName, Type* DataOut)
+	{
+		return TryGetDataOfType(EntryName, Type::StaticClass(), DataOut);
+	}
+
+	template <typename Type>
+	Type GetDataWithDefault(const FString& EntryName, const Type& DefaultValue)
+	{
+		Type Data;
+		if (TryGetData(EntryName, Data))
+		{
+			return Data;
+		}
+
+		return DefaultValue;
+	}
+}
\ No newline at end of file
diff --git a/Engine/Plugins/Experimental/GlobalConfigurationData/Source/GlobalConfigurationData/Public/GlobalConfigurationDataBlueprintLibrary.h b/Engine/Plugins/Experimental/GlobalConfigurationData/Source/GlobalConfigurationData/Public/GlobalConfigurationDataBlueprintLibrary.h
new file mode 100644
--- /dev/null
+++ b/Engine/Plugins/Experimental/GlobalConfigurationData/Source/GlobalConfigurationData/Public/GlobalConfigurationDataBlueprintLibrary.h
@@ -0,0 +1,51 @@
+// Copyright Epic Games, Inc. All Rights Reserved.
+
+#pragma once
+
+#include "CoreMinimal.h"
+#include "Kismet/BlueprintFunctionLibrary.h"
+
+#include "GlobalConfigurationDataBlueprintLibrary.generated.h"
+
+UCLASS()
+class UGlobalConfigurationDataBlueprintLibrary : public UBlueprintFunctionLibrary
+{
+	GENERATED_BODY()
+
+public:
+	UFUNCTION(BlueprintPure, Category="Global Configuration Data")
+	static bool GetConfigDataBool(const FString& EntryName, bool &bValueOut);
+
+	UFUNCTION(BlueprintPure, Category="Global Configuration Data")
+	static bool GetConfigDataInt(const FString& EntryName, int32& ValueOut);
+
+	UFUNCTION(BlueprintPure, Category="Global Configuration Data")
+	static bool GetConfigDataFloat(const FString& EntryName, float& ValueOut);
+
+	UFUNCTION(BlueprintPure, Category="Global Configuration Data")
+	static bool GetConfigDataString(const FString& EntryName, FString& ValueOut);
+
+	UFUNCTION(BlueprintPure, Category="Global Configuration Data")
+	static bool GetConfigDataText(const FString& EntryName, FText& ValueOut);
+
+	UFUNCTION(BlueprintPure, Category="Global Configuration Data")
+	static bool GetConfigDataStruct(const FString& EntryName, UScriptStruct* StructType, FInstancedStruct& ValueOut);
+
+	UFUNCTION(BlueprintPure, Category="Global Configuration Data")
+	static bool GetConfigDataObject(const FString& EntryName, UObject* ValueInOut);
+
+	UFUNCTION(BlueprintPure, Category="Global Configuration Data")
+	static bool GetConfigDataBoolWithDefault(const FString& EntryName, bool bDefaultValue);
+
+	UFUNCTION(BlueprintPure, Category="Global Configuration Data")
+	static int32 GetConfigDataIntWithDefault(const FString& EntryName, int32 DefaultValue);
+
+	UFUNCTION(BlueprintPure, Category="Global Configuration Data")
+	static float GetConfigDataFloatWithDefault(const FString& EntryName, float DefaultValue);
+
+	UFUNCTION(BlueprintPure, Category="Global Configuration Data")
+	static FString GetConfigDataStringWithDefault(const FString& EntryName, FString DefaultValue);
+
+	UFUNCTION(BlueprintPure, Category="Global Configuration Data")
+	static FText GetConfigDataTextWithDefault(const FString& EntryName, FText DefaultValue);
+};
diff --git a/Engine/Plugins/Experimental/GlobalConfigurationData/Source/GlobalConfigurationData/Public/GlobalConfigurationRouter.h b/Engine/Plugins/Experimental/GlobalConfigurationData/Source/GlobalConfigurationData/Public/GlobalConfigurationRouter.h
new file mode 100644
--- /dev/null
+++ b/Engine/Plugins/Experimental/GlobalConfigurationData/Source/GlobalConfigurationData/Public/GlobalConfigurationRouter.h
@@ -0,0 +1,47 @@
+// Copyright Epic Games, Inc. All Rights Reserved.
+#pragma once
+
+#include "Containers/Map.h"
+#include "Misc/Optional.h"
+#include "Templates/SharedPointer.h"
+
+class FJsonValue;
+
+/** The Router is a class that is used to lookup data from various systems in an order determined by their priority
+ * Data should be stored as string values.
+ * 	For primitive types it will use console variable parsing rules.
+ * 	For complex types the data should be stored in a Json format.
+ * 
+ * Routers auto register themselves, so just make sure they exist somewhere and that's it.
+ */
+class IGlobalConfigurationRouter
+{
+public:
+	GLOBALCONFIGURATIONDATA_API IGlobalConfigurationRouter(FString&& InRouterName, int32 InPriority);
+	GLOBALCONFIGURATIONDATA_API virtual ~IGlobalConfigurationRouter();
+
+	GLOBALCONFIGURATIONDATA_API static TSharedPtr<FJsonValue> TryGetData(const FString& EntryName);
+	
+	// Map of Entry Name to a map of Value by Router Name
+	GLOBALCONFIGURATIONDATA_API static void GetAllRegisteredData(TMap<FString, TMap<FString, TSharedRef<FJsonValue>>>& DataOut);
+
+	GLOBALCONFIGURATIONDATA_API static TSharedPtr<FJsonValue> TryParseString(const FStringView& String);
+	GLOBALCONFIGURATIONDATA_API static FString TryPrintString(TSharedPtr<FJsonValue> Value);
+
+protected:
+	virtual TSharedPtr<FJsonValue> TryGetDataFromRouter(const FString& EntryName) const = 0;
+	virtual void GetAllDataFromRouter(TMap<FString, TSharedRef<FJsonValue>>& DataOut) const = 0;
+
+private:
+	FString RouterName;
+
+	/** Arbitrary value used for sorting.
+	 * The built-in routers are at the extreme ends of this range,
+	 * Console Command router being the highest priority (INT32_MAX)
+	 * Config Router being the lowest priority (INT32_MIN)
+	 * 
+	 * The idea is that a config value is baseline and baked into the product where the console command is used for testing or hotfixing values at the very highest.
+	 * User routers should fall in between these two extremes and manage their own priority levels to get the desired outcome.
+	 */ 
+	int32 Priority;
+};
diff --git a/Engine/Plugins/Experimental/GlobalConfigurationData/Source/GlobalConfigurationData/Public/Routers/GlobalConfigurationConfigRouter.h b/Engine/Plugins/Experimental/GlobalConfigurationData/Source/GlobalConfigurationData/Public/Routers/GlobalConfigurationConfigRouter.h
new file mode 100644
--- /dev/null
+++ b/Engine/Plugins/Experimental/GlobalConfigurationData/Source/GlobalConfigurationData/Public/Routers/GlobalConfigurationConfigRouter.h
@@ -0,0 +1,22 @@
+// Copyright Epic Games, Inc. All Rights Reserved.
+#pragma once
+
+#include "GlobalConfigurationRouter.h"
+
+/**
+ * Data Router that pulls from the section GlobalConfigurationData in GEngineIni
+ */
+class FGlobalConfigurationConfigRouter : public IGlobalConfigurationRouter
+{
+public:
+	FGlobalConfigurationConfigRouter();
+	~FGlobalConfigurationConfigRouter();
+
+	virtual TSharedPtr<FJsonValue> TryGetDataFromRouter(const FString& EntryName) const override;
+	virtual void GetAllDataFromRouter(TMap<FString, TSharedRef<FJsonValue>>& DataOut) const override;
+	
+private:
+	void OnConfigSectionsChanged(const FString& IniFilename, const TSet<FString>& SectionNames);
+
+	TMap<FString, TSharedRef<FJsonValue>> StoredData;
+};
diff --git a/Engine/Plugins/Experimental/GlobalConfigurationData/Source/GlobalConfigurationData/Public/Routers/GlobalConfigurationConsoleCommandRouter.h b/Engine/Plugins/Experimental/GlobalConfigurationData/Source/GlobalConfigurationData/Public/Routers/GlobalConfigurationConsoleCommandRouter.h
new file mode 100644
--- /dev/null
+++ b/Engine/Plugins/Experimental/GlobalConfigurationData/Source/GlobalConfigurationData/Public/Routers/GlobalConfigurationConsoleCommandRouter.h
@@ -0,0 +1,18 @@
+// Copyright Epic Games, Inc. All Rights Reserved.
+#pragma once
+
+#include "GlobalConfigurationRouter.h"
+
+/**
+ * Data router that uses console commands to register data with the system:
+ * GCD.RegisterValue Name "Value"
+ * GCD.UnregisterValue Name
+ */
+class FGlobalConfigurationConsoleCommandRouter : public IGlobalConfigurationRouter
+{
+public:
+	FGlobalConfigurationConsoleCommandRouter();
+
+	virtual TSharedPtr<FJsonValue> TryGetDataFromRouter(const FString& EntryName) const override;
+	virtual void GetAllDataFromRouter(TMap<FString, TSharedRef<FJsonValue>>& DataOut) const override;
+};
diff --git a/Engine/Plugins/Experimental/GlobalConfigurationData/Tests/GlobalConfigurationDataTests.Build.cs b/Engine/Plugins/Experimental/GlobalConfigurationData/Tests/GlobalConfigurationDataTests.Build.cs
new file mode 100644
--- /dev/null
+++ b/Engine/Plugins/Experimental/GlobalConfigurationData/Tests/GlobalConfigurationDataTests.Build.cs
@@ -0,0 +1,25 @@
+// Copyright Epic Games, Inc. All Rights Reserved.
+
+using UnrealBuildTool;
+
+public class GlobalConfigurationDataTests : TestModuleRules
+{
+	static GlobalConfigurationDataTests()
+	{
+		TestMetadata = new Metadata()
+		{
+			TestName = "GlobalConfigurationData",
+			TestShortName = "GlobalConfigurationData"
+		};
+	}
+
+	public GlobalConfigurationDataTests(ReadOnlyTargetRules Target) : base(Target)
+	{
+		PrivateDependencyModuleNames.AddRange(
+			new string[] {
+				"Core",
+				"CoreUObject",
+				"GlobalConfigurationData"
+			});
+	}
+}
\ No newline at end of file
diff --git a/Engine/Plugins/Experimental/GlobalConfigurationData/Tests/GlobalConfigurationDataTests.Target.cs b/Engine/Plugins/Experimental/GlobalConfigurationData/Tests/GlobalConfigurationDataTests.Target.cs
new file mode 100644
--- /dev/null
+++ b/Engine/Plugins/Experimental/GlobalConfigurationData/Tests/GlobalConfigurationDataTests.Target.cs
@@ -0,0 +1,11 @@
+// Copyright Epic Games, Inc. All Rights Reserved.
+
+using UnrealBuildTool;
+
+[SupportedPlatforms(UnrealPlatformClass.All)]
+public class GlobalConfigurationDataTestsTarget : TestTargetRules
+{
+	public GlobalConfigurationDataTestsTarget(TargetInfo Target) : base(Target)
+	{
+	}
+}
diff --git a/Engine/Plugins/Experimental/GlobalConfigurationData/Tests/Private/GlobalConfigurationDataTests.cpp b/Engine/Plugins/Experimental/GlobalConfigurationData/Tests/Private/GlobalConfigurationDataTests.cpp
new file mode 100644
--- /dev/null
+++ b/Engine/Plugins/Experimental/GlobalConfigurationData/Tests/Private/GlobalConfigurationDataTests.cpp
@@ -0,0 +1,148 @@
+// Copyright Epic Games, Inc. All Rights Reserved.
+
+#include "TestHarness.h"
+
+#include "GlobalConfigurationData.h"
+#include "GlobalConfigurationRouter.h"
+#include "GlobalConfigurationTestData.h"
+
+#include UE_INLINE_GENERATED_CPP_BY_NAME(GlobalConfigurationTestData)
+
+class FGlobalConfigurationTestRouter : public IGlobalConfigurationRouter
+{
+public:
+	FGlobalConfigurationTestRouter() : IGlobalConfigurationRouter("Test", 0)
+	{
+	}
+
+	virtual TSharedPtr<FJsonValue> TryGetDataFromRouter(const FString& Name) const override
+	{
+		if (const TSharedRef<FJsonValue>* Value = Data.Find(Name))
+		{
+			return *Value;
+		}
+		return {};
+	}
+
+	virtual void GetAllDataFromRouter(TMap<FString, TSharedRef<FJsonValue>>& DataOut) const override
+	{
+		DataOut = Data;
+	}
+
+	void AddValue(const FStringView& Key, const FStringView& Value)
+	{
+		Data.Add(FString(Key), IGlobalConfigurationRouter::TryParseString(Value).ToSharedRef());
+	}
+
+	TMap<FString, TSharedRef<FJsonValue>> Data;
+};
+
+using namespace UE::GlobalConfigurationData;
+
+TEST_CASE("GlobalConfigurationTest", "All")
+{
+	FGlobalConfigurationTestRouter TestRouter;
+
+	TestRouter.AddValue(TEXT("BoolTrueValue"), TEXT("true"));
+	TestRouter.AddValue(TEXT("BoolFalseValue"), TEXT("false"));
+
+	TestRouter.AddValue(TEXT("IntValue"), TEXT("1"));
+	TestRouter.AddValue(TEXT("FloatValue"), TEXT("1.5"));
+	TestRouter.AddValue(TEXT("StringValue"), TEXT("foobar"));
+
+	SECTION("Find Data")
+	{
+		bool bBoolValue = false;
+		int32 IntValue = 0;
+		float FloatValue = 0;
+		FString StringValue;
+
+		CHECK(TryGetData(TEXT("BoolTrueValue"), bBoolValue));
+		CHECK(bBoolValue == true);
+		CHECK(TryGetData(TEXT("BoolFalseValue"), bBoolValue));
+		CHECK(bBoolValue == false);
+		CHECK(TryGetData(TEXT("IntValue"), IntValue));
+		CHECK(IntValue == 1);
+		CHECK(TryGetData(TEXT("FloatValue"), FloatValue));
+		CHECK(FMath::IsNearlyEqual(FloatValue, 1.5f));
+		CHECK(TryGetData(TEXT("StringValue"), StringValue));
+		CHECK(StringValue == TEXT("foobar"));
+	}
+
+	SECTION("Missing Data")
+	{
+		bool bBoolValue = false;
+		int32 IntValue = 0;
+		float FloatValue = 0;
+		FString StringValue;
+
+		CHECK(TryGetData(TEXT("Gobble"), bBoolValue) == false);
+		CHECK(bBoolValue == false);
+		CHECK(TryGetData(TEXT("Gooble"), IntValue) == false);
+		CHECK(IntValue == 0);
+		CHECK(TryGetData(TEXT("Gooooo"), FloatValue) == false);
+		CHECK(FloatValue == 0);
+		CHECK(TryGetData(TEXT("Gowowow"), StringValue) == false);
+		CHECK(StringValue.IsEmpty());
+	}
+
+	SECTION("Default Data")
+	{
+		CHECK(GetDataWithDefault(TEXT("BoolTrueValue"), false) == true);
+		CHECK(GetDataWithDefault(TEXT("BoolFalseValue"), true) == false);
+		CHECK(GetDataWithDefault<int32>(TEXT("IntValue"), 0) == 1);
+		CHECK(FMath::IsNearlyEqual(GetDataWithDefault(TEXT("FloatValue"), 0.f), 1.5f));
+		CHECK(GetDataWithDefault<FString>(TEXT("StringValue"), {}) == TEXT("foobar"));
+
+		CHECK(GetDataWithDefault(TEXT("Garbage"), false) == false);
+		CHECK(GetDataWithDefault(TEXT("Garbage"), true) == true);
+		CHECK(GetDataWithDefault<int32>(TEXT("Garbage"), 123) == 123);
+		CHECK(FMath::IsNearlyEqual(GetDataWithDefault(TEXT("Garbage"), 5.f), 5.f));
+		CHECK(GetDataWithDefault<FString>(TEXT("Garbage"), {}) == FString());
+	}
+	
+	SECTION("Data Conversions")
+	{
+		CHECK(GetDataWithDefault(TEXT("IntValue"), false) == true);
+		CHECK(FMath::IsNearlyEqual(GetDataWithDefault(TEXT("IntValue"), 0.f), 1.f));
+	}
+
+	TestRouter.AddValue(TEXT("ObjectValue"), TEXT("{\"bBoolValue\":true,\"IntValue\":5,\"IntValueArray\":[1,2,3]}"));
+
+	SECTION("UStructs")
+	{
+		FGlobalConfigurationTestStruct TestStruct;
+		CHECK(TryGetData(TEXT("ObjectValue"), TestStruct));
+		CHECK(TestStruct.bBoolValue == true);
+		CHECK(TestStruct.IntValue == 5);
+		CHECK(TestStruct.IntValueArray.Num() == 3);
+		if (TestStruct.IntValueArray.Num() == 3)
+		{
+			CHECK(TestStruct.IntValueArray[0] == 1);
+			CHECK(TestStruct.IntValueArray[1] == 2);
+			CHECK(TestStruct.IntValueArray[2] == 3);
+		}
+	}
+
+	SECTION("UObjects")
+	{
+		UGlobalConfigurationTestObject* TestObject = NewObject<UGlobalConfigurationTestObject>();
+
+		CHECK(TryGetData(TEXT("ObjectValue"), TestObject));
+		CHECK(TestObject->bBoolValue == true);
+		CHECK(TestObject->IntValue == 5);
+		CHECK(TestObject->IntValueArray.Num() == 3);
+		if (TestObject->IntValueArray.Num() == 3)
+		{
+			CHECK(TestObject->IntValueArray[0] == 1);
+			CHECK(TestObject->IntValueArray[1] == 2);
+			CHECK(TestObject->IntValueArray[2] == 3);
+		}
+	}
+
+	SECTION("SubProperties")
+	{
+		CHECK(GetDataWithDefault(TEXT("ObjectValue.bBoolValue"), false) == true);
+		CHECK(GetDataWithDefault(TEXT("ObjectValue.IntValue"), 0) == 5);
+	}
+}
diff --git a/Engine/Plugins/Experimental/GlobalConfigurationData/Tests/Private/GlobalConfigurationRouterTests.cpp b/Engine/Plugins/Experimental/GlobalConfigurationData/Tests/Private/GlobalConfigurationRouterTests.cpp
new file mode 100644
--- /dev/null
+++ b/Engine/Plugins/Experimental/GlobalConfigurationData/Tests/Private/GlobalConfigurationRouterTests.cpp
@@ -0,0 +1,65 @@
+// Copyright Epic Games, Inc. All Rights Reserved.
+
+#include "TestHarness.h"
+
+#include "GlobalConfigurationData.h"
+#include "GlobalConfigurationTestData.h"
+
+#include "HAL/IConsoleManager.h"
+#include "Misc/ConfigCacheIni.h"
+#include "Misc/CoreDelegates.h"
+
+using namespace UE::GlobalConfigurationData;
+
+TEST_CASE("GlobalConfigurationRouterTest", "Routers")
+{
+	SECTION("Config")
+	{
+		const TCHAR* Section = TEXT("GlobalConfigurationData");
+
+		CHECK(GetDataWithDefault(TEXT("BoolValue"), false) == false);
+		CHECK(GetDataWithDefault<int32>(TEXT("IntValue"), 0) == 0);
+		CHECK(GetDataWithDefault<FGlobalConfigurationTestStruct>(TEXT("StructValue"), {}).IntValue == 0);
+
+		GConfig->SetString(Section, TEXT("BoolValue"), TEXT("1"), GEngineIni);
+		GConfig->SetString(Section, TEXT("IntValue"), TEXT("5"), GEngineIni);
+		GConfig->SetString(Section, TEXT("StructValue"), TEXT("{\"bBoolValue\":true,\"IntValue\":5,\"IntValueArray\":[1,2,3]}"), GEngineIni);
+		FCoreDelegates::TSOnConfigSectionsChanged().Broadcast(GEngineIni, { Section });
+
+		CHECK(GetDataWithDefault(TEXT("BoolValue"), false) == true);
+		CHECK(GetDataWithDefault<int32>(TEXT("IntValue"), 0) == 5);
+		CHECK(GetDataWithDefault<FGlobalConfigurationTestStruct>(TEXT("StructValue"), {}).IntValue == 5);
+
+		GConfig->RemoveKey(Section, TEXT("BoolValue"), GEngineIni);
+		GConfig->RemoveKey(Section, TEXT("IntValue"), GEngineIni);
+		GConfig->RemoveKey(Section, TEXT("StructValue"), GEngineIni);
+		FCoreDelegates::TSOnConfigSectionsChanged().Broadcast(GEngineIni, { Section });
+
+		CHECK(GetDataWithDefault(TEXT("BoolValue"), false) == false);
+		CHECK(GetDataWithDefault<int32>(TEXT("IntValue"), 0) == 0);
+		CHECK(GetDataWithDefault<FGlobalConfigurationTestStruct>(TEXT("StructValue"), {}).IntValue == 0);
+	}
+
+	SECTION("Console Command")
+	{
+		CHECK(GetDataWithDefault(TEXT("BoolValue"), false) == false);
+		CHECK(GetDataWithDefault<int32>(TEXT("IntValue"), 0) == 0);
+		CHECK(GetDataWithDefault<FGlobalConfigurationTestStruct>(TEXT("StructValue"), {}).IntValue == 0);
+
+		IConsoleManager::Get().ProcessUserConsoleInput(TEXT("GCD.RegisterValue BoolValue 1"), *GLog, nullptr);
+		IConsoleManager::Get().ProcessUserConsoleInput(TEXT("GCD.RegisterValue IntValue 5"), *GLog, nullptr);
+		IConsoleManager::Get().ProcessUserConsoleInput(TEXT("GCD.RegisterValue StructValue {\"bBoolValue\":true,\"IntValue\":5,\"IntValueArray\":[1,2,3]}"), *GLog, nullptr);
+
+		CHECK(GetDataWithDefault(TEXT("BoolValue"), false) == true);
+		CHECK(GetDataWithDefault<int32>(TEXT("IntValue"), 0) == 5);
+		CHECK(GetDataWithDefault<FGlobalConfigurationTestStruct>(TEXT("StructValue"), {}).IntValue == 5);
+
+		IConsoleManager::Get().ProcessUserConsoleInput(TEXT("GCD.UnregisterValue BoolValue"), *GLog, nullptr);
+		IConsoleManager::Get().ProcessUserConsoleInput(TEXT("GCD.UnregisterValue IntValue"), *GLog, nullptr);
+		IConsoleManager::Get().ProcessUserConsoleInput(TEXT("GCD.UnregisterValue StructValue"), *GLog, nullptr);
+
+		CHECK(GetDataWithDefault(TEXT("BoolValue"), false) == false);
+		CHECK(GetDataWithDefault<int32>(TEXT("IntValue"), 0) == 0);
+		CHECK(GetDataWithDefault<FGlobalConfigurationTestStruct>(TEXT("StructValue"), {}).IntValue == 0);
+	}
+}
diff --git a/Engine/Plugins/Experimental/GlobalConfigurationData/Tests/Private/GlobalConfigurationTestData.h b/Engine/Plugins/Experimental/GlobalConfigurationData/Tests/Private/GlobalConfigurationTestData.h
new file mode 100644
--- /dev/null
+++ b/Engine/Plugins/Experimental/GlobalConfigurationData/Tests/Private/GlobalConfigurationTestData.h
@@ -0,0 +1,38 @@
+// Copyright Epic Games, Inc. All Rights Reserved.
+#pragma once
+
+#include "UObject/Object.h"
+
+#include "GlobalConfigurationTestData.generated.h"
+
+UCLASS()
+class UGlobalConfigurationTestObject : public UObject
+{
+	GENERATED_BODY()
+
+public:
+	UPROPERTY()
+	bool bBoolValue = false;
+
+	UPROPERTY()
+	int32 IntValue = 0;
+
+	UPROPERTY()
+	TArray<int32> IntValueArray;
+};
+
+USTRUCT()
+struct FGlobalConfigurationTestStruct
+{
+	GENERATED_BODY()
+
+public:
+	UPROPERTY()
+	bool bBoolValue = false;
+
+	UPROPERTY()
+	int32 IntValue = 0;
+
+	UPROPERTY()
+	TArray<int32> IntValueArray;
+};
\ No newline at end of file
diff --git a/Engine/Plugins/Experimental/GlobalConfigurationData/Tests/Private/TestGroupEvents.cpp b/Engine/Plugins/Experimental/GlobalConfigurationData/Tests/Private/TestGroupEvents.cpp
new file mode 100644
--- /dev/null
+++ b/Engine/Plugins/Experimental/GlobalConfigurationData/Tests/Private/TestGroupEvents.cpp
@@ -0,0 +1,26 @@
+// Copyright Epic Games, Inc. All Rights Reserved.
+#include "TestCommon/CoreUtilities.h"
+#include "TestCommon/CoreUObjectUtilities.h"
+#include "Modules/ModuleManager.h"
+
+#include <catch2/catch_test_macros.hpp>
+
+GROUP_BEFORE_GLOBAL(Catch::DefaultGroup)
+{
+	InitTaskGraph();
+	InitThreadPool(true);
+#if WITH_COREUOBJECT
+	InitCoreUObject();
+#endif
+
+	FModuleManager::LoadModulePtr<IModuleInterface>(TEXT("GlobalConfigurationData"));
+}
+
+GROUP_AFTER_GLOBAL(Catch::DefaultGroup)
+{
+#if WITH_COREUOBJECT
+	CleanupCoreUObject();
+#endif
+	CleanupThreadPool();
+	CleanupTaskGraph();
+}
