[GCD] Make console command router debug only, add a 'hotfix' config router for high priority settings

#tests Compiled

diff --git a/Engine/Plugins/Experimental/GlobalConfigurationData/Source/GlobalConfigurationDataCore/Private/GlobalConfigurationDataCoreModule.cpp b/Engine/Plugins/Experimental/GlobalConfigurationData/Source/GlobalConfigurationDataCore/Private/GlobalConfigurationDataCoreModule.cpp
--- a/Engine/Plugins/Experimental/GlobalConfigurationData/Source/GlobalConfigurationDataCore/Private/GlobalConfigurationDataCoreModule.cpp
+++ b/Engine/Plugins/Experimental/GlobalConfigurationData/Source/GlobalConfigurationDataCore/Private/GlobalConfigurationDataCoreModule.cpp
@@ -11,11 +11,21 @@ DEFINE_LOG_CATEGORY(LogGlobalConfigurationData);
 
 class FGlobalConfigurationDataCoreModule : public IModuleInterface
 {
-private:
+public:
+	FGlobalConfigurationDataCoreModule()
+	: LowPriorityConfigRouter(TEXT("GlobalConfigurationData"), TEXT("ConfigLow"), INT32_MIN)
+	, HighPriorityConfigRouter(TEXT("GlobalConfigurationDataHotfix"), TEXT("ConfigHigh"), INT32_MAX-1) // Console command still wins since it's for debugging
+	{
+	}
 	
+private:
 	// Just existing is enough for these
-	FGlobalConfigurationConfigRouter ConfigRouter;
+	FGlobalConfigurationConfigRouter LowPriorityConfigRouter;
+	FGlobalConfigurationConfigRouter HighPriorityConfigRouter;
+
+#if !UE_BUILD_SHIPPING
 	FGlobalConfigurationConsoleCommandRouter ConsoleCommandRouter;
+#endif
 };
 
 IMPLEMENT_MODULE(FGlobalConfigurationDataCoreModule, GlobalConfigurationDataCore);
diff --git a/Engine/Plugins/Experimental/GlobalConfigurationData/Source/GlobalConfigurationDataCore/Private/Routers/GlobalConfigurationConfigRouter.cpp b/Engine/Plugins/Experimental/GlobalConfigurationData/Source/GlobalConfigurationDataCore/Private/Routers/GlobalConfigurationConfigRouter.cpp
--- a/Engine/Plugins/Experimental/GlobalConfigurationData/Source/GlobalConfigurationDataCore/Private/Routers/GlobalConfigurationConfigRouter.cpp
+++ b/Engine/Plugins/Experimental/GlobalConfigurationData/Source/GlobalConfigurationDataCore/Private/Routers/GlobalConfigurationConfigRouter.cpp
@@ -6,13 +6,9 @@
 #include "Misc/ConfigCacheIni.h"
 #include "Misc/CoreDelegates.h"
 
-namespace GlobalConfigurationConfigRouter_Private
-{
-	const TCHAR* Section = TEXT("GlobalConfigurationData");
-}
-
-FGlobalConfigurationConfigRouter::FGlobalConfigurationConfigRouter()
-: IGlobalConfigurationRouter(TEXT("Config"), INT32_MIN)
+FGlobalConfigurationConfigRouter::FGlobalConfigurationConfigRouter(FString&& InConfigSection, FString&& InRouterName, int32 InPriority)
+: IGlobalConfigurationRouter(MoveTemp(InRouterName), InPriority)
+, ConfigSection(MoveTemp(InConfigSection))
 {
 	OnConfigSectionsChanged({}, {});
 	FCoreDelegates::TSOnConfigSectionsChanged().AddRaw(this, &FGlobalConfigurationConfigRouter::OnConfigSectionsChanged);
@@ -39,7 +35,7 @@ void FGlobalConfigurationConfigRouter::GetAllDataFromRouter(TMap<FString, TShare
 
 void FGlobalConfigurationConfigRouter::OnConfigSectionsChanged(const FString& IniFilename, const TSet<FString>& SectionNames)
 {
-	if (SectionNames.IsEmpty() || (IniFilename == GEngineIni && SectionNames.Contains(GlobalConfigurationConfigRouter_Private::Section)))
+	if (SectionNames.IsEmpty() || (IniFilename == GEngineIni && SectionNames.Contains(ConfigSection)))
 	{
 		StoredData.Empty();
 
@@ -51,6 +47,6 @@ void FGlobalConfigurationConfigRouter::OnConfigSectionsChanged(const FString& In
 				}
 			});
 
-		GConfig->ForEachEntry(Visitor, GlobalConfigurationConfigRouter_Private::Section, GEngineIni);
+		GConfig->ForEachEntry(Visitor, *ConfigSection, GEngineIni);
 	}
 }
\ No newline at end of file
diff --git a/Engine/Plugins/Experimental/GlobalConfigurationData/Source/GlobalConfigurationDataCore/Public/Routers/GlobalConfigurationConfigRouter.h b/Engine/Plugins/Experimental/GlobalConfigurationData/Source/GlobalConfigurationDataCore/Public/Routers/GlobalConfigurationConfigRouter.h
--- a/Engine/Plugins/Experimental/GlobalConfigurationData/Source/GlobalConfigurationDataCore/Public/Routers/GlobalConfigurationConfigRouter.h
+++ b/Engine/Plugins/Experimental/GlobalConfigurationData/Source/GlobalConfigurationDataCore/Public/Routers/GlobalConfigurationConfigRouter.h
@@ -9,7 +9,7 @@
 class FGlobalConfigurationConfigRouter : public IGlobalConfigurationRouter
 {
 public:
-	FGlobalConfigurationConfigRouter();
+	FGlobalConfigurationConfigRouter(FString&& InConfigSection, FString&& InRouterName, int32 InPriority);
 	~FGlobalConfigurationConfigRouter();
 
 	virtual TSharedPtr<FJsonValue> TryGetDataFromRouter(const FString& EntryName) const override;
@@ -18,5 +18,6 @@ public:
 private:
 	void OnConfigSectionsChanged(const FString& IniFilename, const TSet<FString>& SectionNames);
 
+	FString ConfigSection;
 	TMap<FString, TSharedRef<FJsonValue>> StoredData;
 };
