[GCD] Add support to auto flatten json objects with a single entry
* Support existing requests that request the field name by just returning the resulting value with a log

#tests Checked that it all works as expected
#DisabledValidations nonUnityAndStaticAnalysis

diff --git a/Engine/Plugins/Experimental/GlobalConfigurationData/Source/GlobalConfigurationDataCore/Private/GlobalConfigurationRouter.cpp b/Engine/Plugins/Experimental/GlobalConfigurationData/Source/GlobalConfigurationDataCore/Private/GlobalConfigurationRouter.cpp
--- a/Engine/Plugins/Experimental/GlobalConfigurationData/Source/GlobalConfigurationDataCore/Private/GlobalConfigurationRouter.cpp
+++ b/Engine/Plugins/Experimental/GlobalConfigurationData/Source/GlobalConfigurationDataCore/Private/GlobalConfigurationRouter.cpp
@@ -37,6 +37,12 @@ namespace IGlobalConfigurationRouter_Private
 		TEXT("GCD.ShowParseErrorsAsWarning"),
 		bShowParseErrorsAsWarning,
 		TEXT("If enabled, json parse errors will be logged as warnings, otherwise they are logged as verbose."));
+
+	static bool bAllowFlattenJsonObject = true;
+	static FAutoConsoleVariableRef CVarAllowFlattenJsonObject(
+		TEXT("GCD.AllowFlattenJsonObject"),
+		bAllowFlattenJsonObject,
+		TEXT("If enabled, assume that a json object with a single entry was auto flattened and accept the resulting value anyway."));
 }
 
 IGlobalConfigurationRouter::IGlobalConfigurationRouter(FString&& InRouterName, int32 InPriority)
@@ -91,6 +97,11 @@ TSharedPtr<FJsonValue> IGlobalConfigurationRouter::TryGetData(const FString& Ful
 						continue;
 					}
 				}
+				else if (GetAllowFlattenJsonObject())
+				{
+					UE_LOGFMT(LogGlobalConfigurationData, Verbose, "IGlobalConfigurationRouter::TryGetData found {EntryName} but it is not an object that can satisfy field request {FieldName}, assuming flattened data and returning anyway {RouterName}: {EntryValue}", EntryName, FieldName, Router->RouterName, TryPrintString(Result));
+					return Result;
+				}
 				else
 				{
 					UE_LOGFMT(LogGlobalConfigurationData, Warning, "IGlobalConfigurationRouter::TryGetData found {EntryName} but it is not an object that can satisfy field request {FieldName}, skipping for router {RouterName}: {EntryValue}", EntryName, FieldName, Router->RouterName, TryPrintString(Result));
@@ -103,8 +114,6 @@ TSharedPtr<FJsonValue> IGlobalConfigurationRouter::TryGetData(const FString& Ful
 		}
 	}
 
-	// Todo: Fall back to searching full value without looking for an entry in an object?
-
 	UE_LOGFMT(LogGlobalConfigurationData, VeryVerbose, "IGlobalConfigurationRouter::TryGetData found no results for {EntryName}", FullEntryName);
 	return {};
 }
@@ -167,3 +176,8 @@ FString IGlobalConfigurationRouter::TryPrintString(TSharedPtr<FJsonValue> Value)
 	FJsonSerializer::Serialize(Value, FString(), Writer);
 	return OutputString;
 }
+
+bool IGlobalConfigurationRouter::GetAllowFlattenJsonObject()
+{
+	return IGlobalConfigurationRouter_Private::bAllowFlattenJsonObject;
+}
diff --git a/Engine/Plugins/Experimental/GlobalConfigurationData/Source/GlobalConfigurationDataCore/Public/GlobalConfigurationRouter.h b/Engine/Plugins/Experimental/GlobalConfigurationData/Source/GlobalConfigurationDataCore/Public/GlobalConfigurationRouter.h
--- a/Engine/Plugins/Experimental/GlobalConfigurationData/Source/GlobalConfigurationDataCore/Public/GlobalConfigurationRouter.h
+++ b/Engine/Plugins/Experimental/GlobalConfigurationData/Source/GlobalConfigurationDataCore/Public/GlobalConfigurationRouter.h
@@ -29,6 +29,9 @@ public:
 	GLOBALCONFIGURATIONDATACORE_API static FString TryPrintString(TSharedPtr<FJsonValue> Value);
 
 protected:
+	// If enabled allow routers to flatten json objects to a single value. The core system will allow field name references to accept single value returns.
+	GLOBALCONFIGURATIONDATACORE_API static bool GetAllowFlattenJsonObject();
+	
 	virtual TSharedPtr<FJsonValue> TryGetDataFromRouter(const FString& EntryName) const = 0;
 	virtual void GetAllDataFromRouter(TMap<FString, TSharedRef<FJsonValue>>& DataOut) const = 0;
 
